Add icons to Razer Chroma plugin.

#rnx

diff --git a/Engine/Plugins/Experimental/RazerChromaDevices/Source/RazerChromaEditor/Source/Private/RazerChromaEditorModule.cpp b/Engine/Plugins/Experimental/RazerChromaDevices/Source/RazerChromaEditor/Source/Private/RazerChromaEditorModule.cpp
--- a/Engine/Plugins/Experimental/RazerChromaDevices/Source/RazerChromaEditor/Source/Private/RazerChromaEditorModule.cpp
+++ b/Engine/Plugins/Experimental/RazerChromaDevices/Source/RazerChromaEditor/Source/Private/RazerChromaEditorModule.cpp
@@ -27,8 +27,8 @@ public:
 		const FVector2D Icon16 = FVector2D(16.0, 16.0);
 		const FVector2D Icon64 = FVector2D(64.0, 64.0);
 		
-		Set("ClassIcon.RazerChromaAnimationAsset", new IMAGE_BRUSH_SVG("Icons/RazerChromaAnimation_16", Icon16));
-		Set("ClassThumbnail.RazerChromaAnimationAsset", new IMAGE_BRUSH_SVG("Icons/RazerChromaAnimation_64", Icon64));
+		Set("ClassIcon.RazerChromaAnimationAsset", new IMAGE_BRUSH_SVG("Icons/ChromaAnimation_16", Icon16));
+		Set("ClassThumbnail.RazerChromaAnimationAsset", new IMAGE_BRUSH_SVG("Icons/ChromaAnimation_64", Icon64));
 	}
 };
 
