Fix a razer chroma crash that would happen if the modular feature was not registered correctly

#pf 66564a838a6af08dc2370ce8

diff --git a/Engine/Plugins/Experimental/RazerChromaDevices/Source/RazerChromaDevices/Private/RazerChromaDevicesModule.cpp b/Engine/Plugins/Experimental/RazerChromaDevices/Source/RazerChromaDevices/Private/RazerChromaDevicesModule.cpp
--- a/Engine/Plugins/Experimental/RazerChromaDevices/Source/RazerChromaDevices/Private/RazerChromaDevicesModule.cpp
+++ b/Engine/Plugins/Experimental/RazerChromaDevices/Source/RazerChromaDevices/Private/RazerChromaDevicesModule.cpp
@@ -115,16 +115,19 @@ namespace UE::RazerChroma
 		TEXT("Forcibly reinitalizes the Razer Chroma Editor API (calls Uninit, and then Init)."),
 		FConsoleCommandDelegate::CreateLambda([]()
 			{
-				FRazerChromaDeviceModule::Get().ForceReinitalize();
+				if (FRazerChromaDeviceModule* Module = FRazerChromaDeviceModule::Get())
+				{
+					Module->ForceReinitalize();
+				}
 			})
 	);
 
 #endif //#if RAZER_CHROMA_SUPPORT
 }
 
-FRazerChromaDeviceModule& FRazerChromaDeviceModule::Get()
+FRazerChromaDeviceModule* FRazerChromaDeviceModule::Get()
 {
-	return IModularFeatures::Get().GetModularFeature<FRazerChromaDeviceModule>(UE::RazerChroma::FeatureName);
+	return (FRazerChromaDeviceModule*)IModularFeatures::Get().GetModularFeatureImplementation(UE::RazerChroma::FeatureName, 0);
 }
 
 FName FRazerChromaDeviceModule::GetModularFeatureName()
@@ -303,7 +306,13 @@ void FRazerChromaDeviceModule::ForceReinitalize()
 
 bool FRazerChromaDeviceModule::IsChromaRuntimeAvailable()
 {
-	return FRazerChromaDeviceModule::Get().IsChromaAvailable();
+	FRazerChromaDeviceModule* Module = FRazerChromaDeviceModule::Get(); 
+	if (!Module)
+	{
+		return false;
+	}
+	
+	return Module->IsChromaAvailable();
 }
 
 const int32 FRazerChromaDeviceModule::FindOrLoadAnimationData(const URazerChromaAnimationAsset* AnimAsset)
@@ -354,4 +363,4 @@ const int32 FRazerChromaDeviceModule::FindOrLoadAnimationData(const FString& Ani
 
 #endif	// #if RAZER_CHROMA_SUPPORT
 
-IMPLEMENT_MODULE(FRazerChromaDeviceModule, RazerChromaDevices)
\ No newline at end of file
+IMPLEMENT_MODULE(FRazerChromaDeviceModule, RazerChromaDevices)
diff --git a/Engine/Plugins/Experimental/RazerChromaDevices/Source/RazerChromaDevices/Private/RazerChromaDevicesModule.h b/Engine/Plugins/Experimental/RazerChromaDevices/Source/RazerChromaDevices/Private/RazerChromaDevicesModule.h
--- a/Engine/Plugins/Experimental/RazerChromaDevices/Source/RazerChromaDevices/Private/RazerChromaDevicesModule.h
+++ b/Engine/Plugins/Experimental/RazerChromaDevices/Source/RazerChromaDevices/Private/RazerChromaDevicesModule.h
@@ -12,7 +12,7 @@ class URazerChromaAnimationAsset;
 class RAZERCHROMADEVICES_API FRazerChromaDeviceModule : public IInputDeviceModule
 {
 public:
-	static FRazerChromaDeviceModule& Get();
+	static FRazerChromaDeviceModule* Get();
 	
 	static FName GetModularFeatureName();
 
diff --git a/Engine/Plugins/Experimental/RazerChromaDevices/Source/RazerChromaDevices/Private/RazerChromaFunctionLibrary.cpp b/Engine/Plugins/Experimental/RazerChromaDevices/Source/RazerChromaDevices/Private/RazerChromaFunctionLibrary.cpp
--- a/Engine/Plugins/Experimental/RazerChromaDevices/Source/RazerChromaDevices/Private/RazerChromaFunctionLibrary.cpp
+++ b/Engine/Plugins/Experimental/RazerChromaDevices/Source/RazerChromaDevices/Private/RazerChromaFunctionLibrary.cpp
@@ -42,7 +42,13 @@ bool URazerChromaFunctionLibrary::PlayChromaAnimation(const URazerChromaAnimatio
 		return false;
 	}
 
-	const int32 LoadedAnimId = FRazerChromaDeviceModule::Get().FindOrLoadAnimationData(AnimToPlay);
+	FRazerChromaDeviceModule* Module = FRazerChromaDeviceModule::Get(); 
+	if (!Module)
+	{
+		return false;
+	}
+	
+	const int32 LoadedAnimId = Module->FindOrLoadAnimationData(AnimToPlay);
 	if (LoadedAnimId == INDEX_NONE)
 	{
 		UE_LOG(LogRazerChroma, Error, TEXT("[%hs] Failed to load Chroma Animation %s"), __func__, *GetNameSafe(AnimToPlay));
@@ -74,7 +80,13 @@ bool URazerChromaFunctionLibrary::IsAnimationPlaying(const URazerChromaAnimation
 		return false;
 	}
 
-	const int32 LoadedAnimId = FRazerChromaDeviceModule::Get().FindOrLoadAnimationData(Anim);
+	FRazerChromaDeviceModule* Module = FRazerChromaDeviceModule::Get(); 
+	if (!Module)
+	{
+		return false;
+	}
+
+	const int32 LoadedAnimId = Module->FindOrLoadAnimationData(Anim);
 	if (LoadedAnimId == INDEX_NONE)
 	{
 		UE_LOG(LogRazerChroma, Error, TEXT("[%hs] Failed to load Chroma Animation %s"), __func__, *GetNameSafe(Anim));
@@ -101,7 +113,13 @@ void URazerChromaFunctionLibrary::StopChromaAnimation(const URazerChromaAnimatio
 		return;
 	}
 
-	const int32 LoadedAnimId = FRazerChromaDeviceModule::Get().FindOrLoadAnimationData(AnimToStop);
+	FRazerChromaDeviceModule* Module = FRazerChromaDeviceModule::Get(); 
+	if (!Module)
+	{
+		return;
+	}
+
+	const int32 LoadedAnimId = Module->FindOrLoadAnimationData(AnimToStop);
 	if (LoadedAnimId == INDEX_NONE)
 	{
 		UE_LOG(LogRazerChroma, Error, TEXT("[%hs] Failed to load Chroma Animation %s"), __func__, *GetNameSafe(AnimToStop));
@@ -128,7 +146,13 @@ void URazerChromaFunctionLibrary::PauseChromaAnimation(const URazerChromaAnimati
 		return;
 	}
 
-	const int32 LoadedAnimId = FRazerChromaDeviceModule::Get().FindOrLoadAnimationData(AnimToPause);
+	FRazerChromaDeviceModule* Module = FRazerChromaDeviceModule::Get(); 
+	if (!Module)
+	{
+		return;
+	}
+
+	const int32 LoadedAnimId = Module->FindOrLoadAnimationData(AnimToPause);
 	if (LoadedAnimId == INDEX_NONE)
 	{
 		UE_LOG(LogRazerChroma, Error, TEXT("[%hs] Failed to load Chroma Animation %s"), __func__, *GetNameSafe(AnimToPause));
@@ -154,7 +178,13 @@ bool URazerChromaFunctionLibrary::IsChromaAnimationPaused(const URazerChromaAnim
 		return false;
 	}
 
-	const int32 LoadedAnimId = FRazerChromaDeviceModule::Get().FindOrLoadAnimationData(Anim);
+	FRazerChromaDeviceModule* Module = FRazerChromaDeviceModule::Get(); 
+	if (!Module)
+	{
+		return false;
+	}
+
+	const int32 LoadedAnimId = Module->FindOrLoadAnimationData(Anim);
 	if (LoadedAnimId == INDEX_NONE)
 	{
 		UE_LOG(LogRazerChroma, Error, TEXT("[%hs] Failed to load Chroma Animation %s"), __func__, *GetNameSafe(Anim));
@@ -180,8 +210,14 @@ void URazerChromaFunctionLibrary::ResumeChromaAnimation(const URazerChromaAnimat
 		UE_LOG(LogRazerChroma, Error, TEXT("[%hs] Invalid AnimToResume!"), __func__);
 		return;
 	}
+	
+	FRazerChromaDeviceModule* Module = FRazerChromaDeviceModule::Get(); 
+	if (!Module)
+	{
+		return;
+	}
 
-	const int32 LoadedAnimId = FRazerChromaDeviceModule::Get().FindOrLoadAnimationData(AnimToResume);
+	const int32 LoadedAnimId = Module->FindOrLoadAnimationData(AnimToResume);
 	if (LoadedAnimId == INDEX_NONE)
 	{
 		UE_LOG(LogRazerChroma, Error, TEXT("[%hs] Failed to load Chroma Animation %s"), __func__, *GetNameSafe(AnimToResume));
@@ -218,8 +254,14 @@ void URazerChromaFunctionLibrary::SetIdleAnimation(const URazerChromaAnimationAs
 		return;
 	}
 
+	FRazerChromaDeviceModule* Module = FRazerChromaDeviceModule::Get(); 
+	if (!Module)
+	{
+		return;
+	}
+	
 	// Load the animation
-	const int32 LoadedAnimId = FRazerChromaDeviceModule::Get().FindOrLoadAnimationData(NewIdleAnimation);
+	const int32 LoadedAnimId = Module->FindOrLoadAnimationData(NewIdleAnimation);
 
 	if (LoadedAnimId == INDEX_NONE)
 	{
@@ -345,4 +387,4 @@ FString URazerChromaFunctionLibrary::LexToString(const ERazerChromaDeviceTypes D
 FString URazerChromaFunctionLibrary::Conv_RazerChromaDeviceTypesToString(const int32 DeviceTypes)
 {
 	return URazerChromaFunctionLibrary::LexToString(static_cast<ERazerChromaDeviceTypes>(DeviceTypes));
-}
\ No newline at end of file
+}
diff --git a/Engine/Plugins/Experimental/RazerChromaDevices/Source/RazerChromaDevices/Private/RazerChromaInputDevice.cpp b/Engine/Plugins/Experimental/RazerChromaDevices/Source/RazerChromaDevices/Private/RazerChromaInputDevice.cpp
--- a/Engine/Plugins/Experimental/RazerChromaDevices/Source/RazerChromaDevices/Private/RazerChromaInputDevice.cpp
+++ b/Engine/Plugins/Experimental/RazerChromaDevices/Source/RazerChromaDevices/Private/RazerChromaInputDevice.cpp
@@ -108,8 +108,14 @@ void FRazerChromaInputDevice::HandlePlayAnimationFile(const FRazerChromaPlayAnim
 		UE_LOG(LogRazerChroma, Error, TEXT("[%hs] There is no animation data for chroma effect %s"), __func__, *Property.AnimName);
 		return;
 	}
+
+	FRazerChromaDeviceModule* Module = FRazerChromaDeviceModule::Get(); 
+	if (!Module)
+	{
+		return;
+	}
 	
-	const int32 LoadedAnimId = FRazerChromaDeviceModule::Get().FindOrLoadAnimationData(Property.AnimName, Property.AnimationByteBuffer);
+	const int32 LoadedAnimId = Module->FindOrLoadAnimationData(Property.AnimName, Property.AnimationByteBuffer);
 
 	if (LoadedAnimId == -1)
 	{
