Mark razer chroma function library as BP Cosmetic only. They will only ever be called on a client.

#rnx

diff --git a/Engine/Plugins/Experimental/RazerChromaDevices/Source/RazerChromaDevices/Public/RazerChromaFunctionLibrary.h b/Engine/Plugins/Experimental/RazerChromaDevices/Source/RazerChromaDevices/Public/RazerChromaFunctionLibrary.h
--- a/Engine/Plugins/Experimental/RazerChromaDevices/Source/RazerChromaDevices/Public/RazerChromaFunctionLibrary.h
+++ b/Engine/Plugins/Experimental/RazerChromaDevices/Source/RazerChromaDevices/Public/RazerChromaFunctionLibrary.h
@@ -28,7 +28,7 @@ public:
 	* This will be false on any machines that do not have Razer Chroma installed on them,
 	* and thus cannot set any Razer Chroma effects.
 	*/
-	UFUNCTION(BlueprintCallable, BlueprintPure, Category = "Razer Chroma")
+	UFUNCTION(BlueprintCallable, BlueprintPure, BlueprintCosmetic, Category = "Razer Chroma")
 	static bool IsChromaRuntimeAvailable();
 
 	/**
@@ -41,31 +41,31 @@ public:
 	* 
 	* @return	True if successfully played, false otherwise.
 	*/
-	UFUNCTION(BlueprintCallable, Category = "Razer Chroma", meta = (ReturnDisplayName = "Was Successful"))
+	UFUNCTION(BlueprintCallable, BlueprintCosmetic, Category = "Razer Chroma", meta = (ReturnDisplayName = "Was Successful"))
 	static bool PlayChromaAnimation(const URazerChromaAnimationAsset* AnimToPlay, const bool bLooping = false);
 
 	/**
 	* Returns true if the given animation is currently playing.
 	*/
-	UFUNCTION(BlueprintCallable, BlueprintPure, Category = "Razer Chroma", meta = (ReturnDisplayName = "Is Playing"))
+	UFUNCTION(BlueprintCallable, BlueprintPure, BlueprintCosmetic, Category = "Razer Chroma", meta = (ReturnDisplayName = "Is Playing"))
 	static bool IsAnimationPlaying(const URazerChromaAnimationAsset* Anim);
 
 	/**
 	* Stops the given Chroma Animation if it is currently playing.
 	*/
-	UFUNCTION(BlueprintCallable, Category = "Razer Chroma")
+	UFUNCTION(BlueprintCallable, BlueprintCosmetic, Category = "Razer Chroma")
 	static void StopChromaAnimation(const URazerChromaAnimationAsset* AnimToStop);
 
 	/**
 	* Pauses the given animation if it is currently playing
 	*/
-	UFUNCTION(BlueprintCallable, Category = "Razer Chroma")
+	UFUNCTION(BlueprintCallable, BlueprintCosmetic, Category = "Razer Chroma")
 	static void PauseChromaAnimation(const URazerChromaAnimationAsset* AnimToPause);
 
 	/**
 	 * Returns true if the given animation is currently paused.
 	 */
-	UFUNCTION(BlueprintCallable, BlueprintPure, Category = "Razer Chroma")
+	UFUNCTION(BlueprintCallable, BlueprintPure, BlueprintCosmetic, Category = "Razer Chroma")
 	static bool IsChromaAnimationPaused(const URazerChromaAnimationAsset* Anim);
 	
 	/**
@@ -74,13 +74,13 @@ public:
 	* @param AnimToResume	The animation to resume.
 	* @param bLoop			If true, this animation will loop (start re-playing after it finishes)
 	*/
-	UFUNCTION(BlueprintCallable, Category = "Razer Chroma")
+	UFUNCTION(BlueprintCallable, BlueprintCosmetic, Category = "Razer Chroma")
 	static void ResumeChromaAnimation(const URazerChromaAnimationAsset* AnimToResume, const bool bLoop);
 
 	/**
 	* Stops all currently active Chroma animations
 	*/
-	UFUNCTION(BlueprintCallable, Category = "Razer Chroma")
+	UFUNCTION(BlueprintCallable, BlueprintCosmetic, Category = "Razer Chroma")
 	static void StopAllChromaAnimations();
 
 	/**
@@ -89,13 +89,13 @@ public:
 	* 
 	* By default, the idle animation is set via the project settings, but it can be changed at runtime.
 	*/
-	UFUNCTION(BlueprintCallable, Category = "Razer Chroma")
+	UFUNCTION(BlueprintCallable, BlueprintCosmetic, Category = "Razer Chroma")
 	static void SetIdleAnimation(const URazerChromaAnimationAsset* NewIdleAnimation);
 
 	/**
 	* Sets whether or not we should use an idle currently an animation 
 	*/
-	UFUNCTION(BlueprintCallable, Category = "Razer Chroma")
+	UFUNCTION(BlueprintCallable, BlueprintCosmetic, Category = "Razer Chroma")
 	static void SetUseIdleAnimation(const bool bUseIdleAnimation);
 
 	/**
@@ -112,7 +112,7 @@ public:
 	* @param ColorToSet		The color to set the devices to (the alpha channel is not used)
 	* @param DeviceTypes	Which types of razer devices you would like to set the color on if they are available
 	*/
-	UFUNCTION(BlueprintCallable, Category = "Razer Chroma", meta=(AutoCreateRefTerm="ColorToSet"))
+	UFUNCTION(BlueprintCallable, BlueprintCosmetic, Category = "Razer Chroma", meta=(AutoCreateRefTerm="ColorToSet"))
 	static void SetAllDevicesStaticColor(const FColor& ColorToSet, UPARAM(meta = (Bitmask, BitmaskEnum = "/Script/RazerChromaDevices.ERazerChromaDeviceTypes")) const int32 DeviceTypes);
 
 	/**
