Fix a build health error when running a client which doesn't have razer synapse installed, and thus the razer DLL is not available (i.e. you run on an automated test sever, or just dont have synapse installed). Change this case from an error log to a normal log to prevent error log spam.

#rnx

diff --git a/Engine/Plugins/Experimental/RazerChromaDevices/Source/RazerChromaDevices/Private/RazerChromaDevicesModule.cpp b/Engine/Plugins/Experimental/RazerChromaDevices/Source/RazerChromaDevices/Private/RazerChromaDevicesModule.cpp
--- a/Engine/Plugins/Experimental/RazerChromaDevices/Source/RazerChromaDevices/Private/RazerChromaDevicesModule.cpp
+++ b/Engine/Plugins/Experimental/RazerChromaDevices/Source/RazerChromaDevices/Private/RazerChromaDevicesModule.cpp
@@ -206,10 +206,6 @@ void FRazerChromaDeviceModule::StartupModule()
 
 	// Initialize the SDK	
 	const RZRESULT Res = UE::RazerChroma::InitChromaSDK();
-	UE_CLOG(Res != RZRESULT_SUCCESS, LogRazerChroma, Error, TEXT("[%hs] Failed to Init Razer Chroma Editor API. Error code %d (%s)"),
-			__func__,
-			Res,
-			*FRazerChromaDeviceModule::RazerErrorToString(Res));
 	
 	// Ensure that we keep track of if we have loaded the API successfully or not for later.
 	bLoadedDynamicAPISuccessfully &= (Res == RZRESULT_SUCCESS);
@@ -228,6 +224,24 @@ void FRazerChromaDeviceModule::StartupModule()
 		}
 	}
 
+	// This will be the result if you run on a machine which does not have the Razer Synapse client installed
+	// (i.e. you don't have any razer products) We don't want to error here, as that would be expected. 
+	if (Res == RZRESULT_DLL_NOT_FOUND)
+	{
+		UE_LOG(LogRazerChroma, Log, TEXT("[%hs] Failed to Init Razer Chroma Editor API. The Razer Synapse client is likely not installed on this machine. Error code %d (%s)"),
+			__func__,
+			Res,
+			*FRazerChromaDeviceModule::RazerErrorToString(Res));
+	}
+	// If we failed for any other reason then it is not expected and we should log an error
+	else if (Res != RZRESULT_SUCCESS)
+	{
+		UE_LOG(LogRazerChroma, Error, TEXT("[%hs] Failed to Init Razer Chroma Editor API. Error code %d (%s)"),
+			__func__,
+			Res,
+			*FRazerChromaDeviceModule::RazerErrorToString(Res));
+	}	
+
 #else
 	
 	UE_LOG(LogRazerChroma, Log, TEXT("[%hs] RAZER_CHROMA_SUPPORT=0    No Razer Chroma Animation functionality will be available."), __func__);
