Make sure to stop all animations and clear the idle animation on module shutdown of Razer Chroma

#rnx
#pf 6661f667483bfbf315e75eaf

diff --git a/Engine/Plugins/Experimental/RazerChromaDevices/Source/RazerChromaDevices/Private/RazerChromaDevicesModule.cpp b/Engine/Plugins/Experimental/RazerChromaDevices/Source/RazerChromaDevices/Private/RazerChromaDevicesModule.cpp
--- a/Engine/Plugins/Experimental/RazerChromaDevices/Source/RazerChromaDevices/Private/RazerChromaDevicesModule.cpp
+++ b/Engine/Plugins/Experimental/RazerChromaDevices/Source/RazerChromaDevices/Private/RazerChromaDevicesModule.cpp
@@ -255,13 +255,13 @@ void FRazerChromaDeviceModule::ShutdownModule()
 
 #if RAZER_CHROMA_SUPPORT
 
+	// Modular feature is no longer available
 	IModularFeatures::Get().UnregisterModularFeature(UE::RazerChroma::FeatureName, this);
 
-	if (FRazerChromaEditorDynamicAPI::UnInit)
-	{
-		FRazerChromaEditorDynamicAPI::UnInit();
-	}
-	
+	// Run some razer chroma specific cleanup
+	CleanupSDK();
+
+	// Free the DLL handle from the process
 	if (RazerChromaEditorDLLHandle)
 	{
 		FPlatformProcess::FreeDllHandle(RazerChromaEditorDLLHandle);
@@ -292,6 +292,42 @@ TSharedPtr<IInputDevice> FRazerChromaDeviceModule::CreateInputDevice(const TShar
 }
 
 #if RAZER_CHROMA_SUPPORT
+
+void FRazerChromaDeviceModule::CleanupSDK()
+{
+	// Disable idle animations
+	if (FRazerChromaEditorDynamicAPI::SetUseIdleAnimations)
+	{
+		FRazerChromaEditorDynamicAPI::SetUseIdleAnimations(false);
+	}
+
+	// Stop playing all animations
+	if (FRazerChromaEditorDynamicAPI::StopAllAnimations)
+	{
+		FRazerChromaEditorDynamicAPI::StopAllAnimations();
+	}
+
+	// Return any animations to disk
+	if (FRazerChromaEditorDynamicAPI::CloseAll)
+	{
+		FRazerChromaEditorDynamicAPI::CloseAll();
+	}
+
+	// Finally, UnInit the whole sdk
+	if (FRazerChromaEditorDynamicAPI::UnInit)
+	{
+		FRazerChromaEditorDynamicAPI::UnInit();
+	}
+
+	// Doing all of the above _should_ reset the state of Razer peripherals to the user's
+	// default settings and make sure that the application is correctly removed from Razer Synapse...
+
+	UE_LOG(LogRazerChroma,
+		Log,
+		TEXT("[%hs] Razer Chroma Editor library cleaned up."),
+		__func__);
+}
+
 bool FRazerChromaDeviceModule::IsChromaAvailable() const
 {
 	return RazerChromaEditorDLLHandle != nullptr && bLoadedDynamicAPISuccessfully;
@@ -305,10 +341,7 @@ void FRazerChromaDeviceModule::ForceReinitalize()
 	}
 
 	// Force Uninit...
-	if (FRazerChromaEditorDynamicAPI::UnInit)
-	{
-		FRazerChromaEditorDynamicAPI::UnInit();
-	}
+	CleanupSDK();
 
 	// And re-init
 	const RZRESULT Res = UE::RazerChroma::InitChromaSDK();
diff --git a/Engine/Plugins/Experimental/RazerChromaDevices/Source/RazerChromaDevices/Private/RazerChromaDevicesModule.h b/Engine/Plugins/Experimental/RazerChromaDevices/Source/RazerChromaDevices/Private/RazerChromaDevicesModule.h
--- a/Engine/Plugins/Experimental/RazerChromaDevices/Source/RazerChromaDevices/Private/RazerChromaDevicesModule.h
+++ b/Engine/Plugins/Experimental/RazerChromaDevices/Source/RazerChromaDevices/Private/RazerChromaDevicesModule.h
@@ -38,6 +38,11 @@ protected:
      */
     bool IsChromaAvailable() const;
 
+	/**
+	* Cleans up the SDK and all currently playing animations.
+	*/
+	void CleanupSDK();
+
 	/** Handle to the Razer Chroma dynamic DLL */
 	void* RazerChromaEditorDLLHandle = nullptr;
 
diff --git a/Engine/Plugins/Experimental/RazerChromaDevices/Source/RazerChromaDevices/Private/RazerChromaDynamicAPI.cpp b/Engine/Plugins/Experimental/RazerChromaDevices/Source/RazerChromaDevices/Private/RazerChromaDynamicAPI.cpp
--- a/Engine/Plugins/Experimental/RazerChromaDevices/Source/RazerChromaDevices/Private/RazerChromaDynamicAPI.cpp
+++ b/Engine/Plugins/Experimental/RazerChromaDevices/Source/RazerChromaDevices/Private/RazerChromaDynamicAPI.cpp
@@ -26,6 +26,7 @@ FRazerChromaEditorDynamicAPI::DELETEEFFECT FRazerChromaEditorDynamicAPI::DeleteE
 FRazerChromaEditorDynamicAPI::SETIDLEANIMATION FRazerChromaEditorDynamicAPI::SetIdleAnimation = nullptr;
 FRazerChromaEditorDynamicAPI::USEIDLEANIMATIONS FRazerChromaEditorDynamicAPI::SetUseIdleAnimations = nullptr;
 FRazerChromaEditorDynamicAPI::STOPALLANIMATIONS FRazerChromaEditorDynamicAPI::StopAllAnimations = nullptr;
+FRazerChromaEditorDynamicAPI::CLOSEALL FRazerChromaEditorDynamicAPI::CloseAll = nullptr;
 FRazerChromaEditorDynamicAPI::PAUSEANIMATION FRazerChromaEditorDynamicAPI::PauseAnimation = nullptr;
 FRazerChromaEditorDynamicAPI::RESUMEANIMATION FRazerChromaEditorDynamicAPI::ResumeAnimation = nullptr;
 FRazerChromaEditorDynamicAPI::ISANIMATIONPLAYING FRazerChromaEditorDynamicAPI::IsAnimationPlaying = nullptr;
@@ -65,7 +66,8 @@ bool FRazerChromaEditorDynamicAPI::LoadAPI(void* RazerChromaEditorDLLHandle)
 	FRazerChromaEditorDynamicAPI::DeleteEffect = reinterpret_cast<DELETEEFFECT>(GetProcAddress(RazerModule, "PluginCoreDeleteEffect"));
 	FRazerChromaEditorDynamicAPI::SetIdleAnimation = reinterpret_cast<SETIDLEANIMATION>(GetProcAddress(RazerModule, "PluginSetIdleAnimation"));
 	FRazerChromaEditorDynamicAPI::SetUseIdleAnimations = reinterpret_cast<USEIDLEANIMATIONS>(GetProcAddress(RazerModule, "PluginUseIdleAnimations"));
-	FRazerChromaEditorDynamicAPI::StopAllAnimations = reinterpret_cast<STOPALLANIMATIONS>(GetProcAddress(RazerModule, "PluginStopAll"));
+	FRazerChromaEditorDynamicAPI::StopAllAnimations = reinterpret_cast<STOPALLANIMATIONS>(GetProcAddress(RazerModule, "PluginStopAll"));	
+	FRazerChromaEditorDynamicAPI::CloseAll = reinterpret_cast<CLOSEALL>(GetProcAddress(RazerModule, "PluginCloseAll"));
 	FRazerChromaEditorDynamicAPI::PauseAnimation = reinterpret_cast<PAUSEANIMATION>(GetProcAddress(RazerModule, "PluginPauseAnimation"));
 	FRazerChromaEditorDynamicAPI::ResumeAnimation = reinterpret_cast<RESUMEANIMATION>(GetProcAddress(RazerModule, "PluginResumeAnimation"));
 	FRazerChromaEditorDynamicAPI::IsAnimationPlaying = reinterpret_cast<ISANIMATIONPLAYING>(GetProcAddress(RazerModule, "PluginIsPlaying"));
@@ -93,6 +95,7 @@ bool FRazerChromaEditorDynamicAPI::LoadAPI(void* RazerChromaEditorDLLHandle)
 		FRazerChromaEditorDynamicAPI::SetIdleAnimation && 
 		FRazerChromaEditorDynamicAPI::SetUseIdleAnimations && 
 		FRazerChromaEditorDynamicAPI::StopAllAnimations && 
+		FRazerChromaEditorDynamicAPI::CloseAll &&
 		FRazerChromaEditorDynamicAPI::PauseAnimation && 
 		FRazerChromaEditorDynamicAPI::ResumeAnimation &&
 		FRazerChromaEditorDynamicAPI::IsAnimationPlaying &&
diff --git a/Engine/Plugins/Experimental/RazerChromaDevices/Source/RazerChromaDevices/Private/RazerChromaDynamicAPI.h b/Engine/Plugins/Experimental/RazerChromaDevices/Source/RazerChromaDevices/Private/RazerChromaDynamicAPI.h
--- a/Engine/Plugins/Experimental/RazerChromaDevices/Source/RazerChromaDevices/Private/RazerChromaDynamicAPI.h
+++ b/Engine/Plugins/Experimental/RazerChromaDevices/Source/RazerChromaDevices/Private/RazerChromaDynamicAPI.h
@@ -66,6 +66,8 @@ struct FRazerChromaEditorDynamicAPI
 
 	typedef void (*STOPALLANIMATIONS)();
 
+	typedef void (*CLOSEALL)();
+
 	typedef void (*PAUSEANIMATION)(int32 animationId);
 
 	typedef void (*RESUMEANIMATION)(int32 animationId, bool loop);
@@ -115,6 +117,10 @@ struct FRazerChromaEditorDynamicAPI
 	// Stops all animations that are currently playing
 	static STOPALLANIMATIONS StopAllAnimations;
 
+	// Closes all open animations so they can be reloaded from disk.
+	// You should call this on application shutdown
+	static CLOSEALL CloseAll;
+
 	// Pauses the animation with the given int32 id
 	static PAUSEANIMATION PauseAnimation;
 
