Append the build target config and type to the Razer Application Title name in non-shipping builds.

This will register the different targets (editor, test, shipping, etc) as different apps in razer synapse, which will ensure a clean test enviornment for each target type. If you tell Razer that it is the same "Application Name", but run a different *.exe file, Synapse get's a little fussy and unhappy with you, and will only actually work with the first one you ran.

Char limit is now 236 instead of 256 to make sure that it works with the longest target/build type combination.

#rnx

diff --git a/Engine/Plugins/Experimental/RazerChromaDevices/Source/RazerChromaDevices/Private/RazerChromaDevicesModule.cpp b/Engine/Plugins/Experimental/RazerChromaDevices/Source/RazerChromaDevices/Private/RazerChromaDevicesModule.cpp
--- a/Engine/Plugins/Experimental/RazerChromaDevices/Source/RazerChromaDevices/Private/RazerChromaDevicesModule.cpp
+++ b/Engine/Plugins/Experimental/RazerChromaDevices/Source/RazerChromaDevices/Private/RazerChromaDevicesModule.cpp
@@ -2,9 +2,11 @@
 
 #include "RazerChromaDevicesModule.h"
 
+#include "GenericPlatform/GenericPlatformMisc.h"
 #include "HAL/IConsoleManager.h"
 #include "HAL/PlatformProcess.h"
 #include "Interfaces/IPluginManager.h"
+#include "Misc/App.h"
 #include "Misc/Paths.h"
 #include "RazerChromaDeviceLogging.h"
 #include "RazerChromaDevicesDeveloperSettings.h"
@@ -80,9 +82,33 @@ namespace UE::RazerChroma
 
 			FRazerChromaAppInfo WhyMe = {};
 
-			ChromaSDK::APPINFOTYPE AppInfo = {};
+			ChromaSDK::APPINFOTYPE AppInfo = {};			
+
+			// Make sure that the application name will fit with our build config appends
+			ensure(SettingsAppInfo.ApplicationTitle.Len() <= 236);
+
+			// Outside of shipping builds, we will append the build config and target type to the application name
+			// so that Razer Synapse recognizes them as different apps. This makes testing a little easier and ensures 
+			// that the environment is clean for testing shipping builds.
+	#if !UE_BUILD_SHIPPING
+
+			TStringBuilder<256> TitleBuilder;
+
+			TitleBuilder.Append(SettingsAppInfo.ApplicationTitle);
+			TitleBuilder.Append(TEXT("_"));
+			TitleBuilder.Append(LexToString(FApp::GetBuildConfiguration()));
+			TitleBuilder.Append(TEXT("_"));
+			TitleBuilder.Append(LexToString(FApp::GetBuildTargetType()));
+
+			FCString::Strncpy(AppInfo.Title, TitleBuilder.GetData(), 256);
+
+	#else
 
 			FCString::Strncpy(AppInfo.Title, *SettingsAppInfo.ApplicationTitle, 256);
+
+	#endif	// !UE_BUILD_SHIPPPING			
+
+
 			FCString::Strncpy(AppInfo.Description, *SettingsAppInfo.ApplicationDescription, 1024);
 
 			FCString::Strncpy(AppInfo.Author.Name, *SettingsAppInfo.AuthorName, 256);
diff --git a/Engine/Plugins/Experimental/RazerChromaDevices/Source/RazerChromaDevices/Public/RazerChromaDevicesDeveloperSettings.h b/Engine/Plugins/Experimental/RazerChromaDevices/Source/RazerChromaDevices/Public/RazerChromaDevicesDeveloperSettings.h
--- a/Engine/Plugins/Experimental/RazerChromaDevices/Source/RazerChromaDevices/Public/RazerChromaDevicesDeveloperSettings.h
+++ b/Engine/Plugins/Experimental/RazerChromaDevices/Source/RazerChromaDevices/Public/RazerChromaDevicesDeveloperSettings.h
@@ -36,11 +36,23 @@ struct FRazerChromaAppInfo
 	/**
 	* The name of your application to report to Razer Synapse.
 	* 
+	* In non-shipping builds, this will have the build target and type appended to it, which will make it show up as :
+	* 
+	*	"ApplicationTitle_<BuildType>_<BuildTarget>" in Razer Synapse. 
+	* 
+	* This is expected, and in a shipping build it will be displayed as only: 
+	* 
+	*	"ApplicationTitle"
+	* 
+	* This is necessary because if you register multiple executable names (i.e. <YourGame>.exe, <YourGame>-Win64-Shipping.exe, and <YourGame>-Win64-Test.exe)
+	* then Razer Synapse will only use the _first_ executable that you ran with this configuration. Appending the target names makes them unique in Synapse and easier
+	* to test.
+	* 
 	* Cannot be empty.
 	* 
-	* Char limit of 256.
+	* Char limit of 236.
 	*/
-	UPROPERTY(EditAnywhere, Category="Razer Chroma App Config", NoClear, meta=(MaxLength=256))
+	UPROPERTY(EditAnywhere, Category="Razer Chroma App Config", NoClear, meta=(MaxLength=236))
 	FString ApplicationTitle = TEXT("Your Game Name here");
 
 	/**
