[Razer Chroma plugin]

PR #12032: Fix for the title has  requiring a null terminator

diff --git a/Engine/Plugins/Experimental/RazerChromaDevices/Source/RazerChromaDevices/Private/RazerChromaDevicesModule.cpp b/Engine/Plugins/Experimental/RazerChromaDevices/Source/RazerChromaDevices/Private/RazerChromaDevicesModule.cpp
--- a/Engine/Plugins/Experimental/RazerChromaDevices/Source/RazerChromaDevices/Private/RazerChromaDevicesModule.cpp
+++ b/Engine/Plugins/Experimental/RazerChromaDevices/Source/RazerChromaDevices/Private/RazerChromaDevicesModule.cpp
@@ -100,7 +100,7 @@ namespace UE::RazerChroma
 			TitleBuilder.Append(TEXT("_"));
 			TitleBuilder.Append(LexToString(FApp::GetBuildTargetType()));
 
-			FCString::Strncpy(AppInfo.Title, TitleBuilder.GetData(), 256);
+			FCString::Strncpy(AppInfo.Title, TitleBuilder.ToString(), 256);
 
 	#else
 
