Fix for RazerChroma PDB files not being found in GitHub builds. The PDB files are not strictly necessary for compiling the plugin, so we will not redist them with the engine. If you need the PDB files you can download them from the Razer Chroma editor github.

#rnx

diff --git a/Engine/Build/InstalledEngineFilters.xml b/Engine/Build/InstalledEngineFilters.xml
--- a/Engine/Build/InstalledEngineFilters.xml
+++ b/Engine/Build/InstalledEngineFilters.xml
@@ -154,6 +154,7 @@
 		Engine/Plugins/Runtime/NetworkPrediction/readme.txt
 		Engine/Plugins/Experimental/Mover/README.md
 		Engine/Plugins/Runtime/GameplayAbilities/README.md
+		Engine/Plugins/Experimental/RazerChromaDevices/README.txt
 
 	</Property>
 	
diff --git a/Engine/Plugins/Experimental/RazerChromaDevices/README.txt b/Engine/Plugins/Experimental/RazerChromaDevices/README.txt
new file mode 100644
--- /dev/null
+++ b/Engine/Plugins/Experimental/RazerChromaDevices/README.txt
@@ -0,0 +1,10 @@
+# A note about Razer Chroma PDB files
+
+If you would like to compile with the RazerChroma PDB file for better
+debugging symbols, then you should download them from the Razer CChromaEditor
+github : [REDACTED_PROFILE_URL]/CChromaEditor
+
+You should put the .pdb files in the following directories to build successfully:
+
+Win64: Engine\Plugins\Experimental\RazerChromaDevices\Binaries\ThirdParty\Win64\CChromaEditorLibrary64.pdb
+Win32: Engine\Plugins\Experimental\RazerChromaDevices\Binaries\ThirdParty\Win32\CChromaEditorLibrary.pdb
\ No newline at end of file
diff --git a/Engine/Plugins/Experimental/RazerChromaDevices/Source/RazerChromaSDK/RazerChromaSDK.build.cs b/Engine/Plugins/Experimental/RazerChromaDevices/Source/RazerChromaSDK/RazerChromaSDK.build.cs
--- a/Engine/Plugins/Experimental/RazerChromaDevices/Source/RazerChromaSDK/RazerChromaSDK.build.cs
+++ b/Engine/Plugins/Experimental/RazerChromaDevices/Source/RazerChromaSDK/RazerChromaSDK.build.cs
@@ -27,7 +27,13 @@ namespace UnrealBuildTool.Rules
 					// We only want the debug symbols outside of shipping
 					if (Target.Configuration != UnrealTargetConfiguration.Shipping)
 					{
-						RuntimeDependencies.Add(Path.Combine(DllPath, "CChromaEditorLibrary64.pdb"));
+						// The PDB file may not exist if cloning from GitHub. See the README for
+						// where to put it if you need extra debug symbols.
+						string PDBPath = Path.Combine(DllPath, "CChromaEditorLibrary64.pdb");
+						if (File.Exists(PDBPath))
+						{
+							RuntimeDependencies.Add(PDBPath);
+						}						
 					}					
 				}
 				else
@@ -40,7 +46,13 @@ namespace UnrealBuildTool.Rules
 					// We only want the debug symbols outside of shipping
 					if (Target.Configuration != UnrealTargetConfiguration.Shipping)
 					{
-						RuntimeDependencies.Add(Path.Combine(DllPath, "CChromaEditorLibrary.pdb"));
+						// The PDB file may not exist if cloning from GitHub. See the README for
+						// where to put it if you need extra debug symbols.
+						string PDBPath = Path.Combine(DllPath, "CChromaEditorLibrary.pdb");
+						if (File.Exists(PDBPath))
+						{
+							RuntimeDependencies.Add(PDBPath);
+						}							
 					}						
 				}
 			}
