Change the friendly name of the Razer Chroma Input Devices plugin to just "Razer Chroma Devices"... technically it supports more then just input devices (Different lamps/general peripheral lighting)

#rnx

diff --git a/Engine/Plugins/Experimental/RazerChromaDevices/RazerChromaDevices.uplugin b/Engine/Plugins/Experimental/RazerChromaDevices/RazerChromaDevices.uplugin
--- a/Engine/Plugins/Experimental/RazerChromaDevices/RazerChromaDevices.uplugin
+++ b/Engine/Plugins/Experimental/RazerChromaDevices/RazerChromaDevices.uplugin
@@ -1,12 +1,12 @@
 {
 	"FileVersion": 3,
-	"FriendlyName": "Razer Chroma Input Devices",
+	"FriendlyName": "Razer Chroma Devices",
 	"Version": 1,
 	"VersionName": "1.0",
 	"CreatedBy": "Epic Games, Inc.",
 	"CreatedByURL": "https://epicgames.com",
 	"Description": "Provides some functionality to set Razer Chroma effects at runtime.",
-	"Category": "Input Devices",
+	"Category": "Peripherals",
 	"EnabledByDefault": false,
 	"IsBetaVersion": true,
 	"CanContainContent": true,
