remove win32 case.

#rnx

diff --git a/Engine/Plugins/Experimental/RazerChromaDevices/Source/RazerChromaSDK/RazerChromaSDK.build.cs b/Engine/Plugins/Experimental/RazerChromaDevices/Source/RazerChromaSDK/RazerChromaSDK.build.cs
--- a/Engine/Plugins/Experimental/RazerChromaDevices/Source/RazerChromaSDK/RazerChromaSDK.build.cs
+++ b/Engine/Plugins/Experimental/RazerChromaDevices/Source/RazerChromaSDK/RazerChromaSDK.build.cs
@@ -36,25 +36,6 @@ namespace UnrealBuildTool.Rules
 						}						
 					}					
 				}
-				else
-				{
-					// Redist the CChromaEditorLibrary.dll with the game
-					string DllPath = Path.GetFullPath(Path.Combine(PluginDirectory, "Binaries/ThirdParty/Win32"));
-
-					RuntimeDependencies.Add(Path.Combine(DllPath, "CChromaEditorLibrary.dll"));
-
-					// We only want the debug symbols outside of shipping
-					if (Target.Configuration != UnrealTargetConfiguration.Shipping)
-					{
-						// The PDB file may not exist if cloning from GitHub. See the README for
-						// where to put it if you need extra debug symbols.
-						string PDBPath = Path.Combine(DllPath, "CChromaEditorLibrary.pdb");
-						if (File.Exists(PDBPath))
-						{
-							RuntimeDependencies.Add(PDBPath);
-						}							
-					}						
-				}
 			}
 		}
 	}
