RazerChromaEditor: Move Private directory to conform with standard Module layout

#rnx

diff --git a/Engine/Plugins/Experimental/RazerChromaDevices/Source/RazerChromaEditor/Private/RazerChromaAnimationAssetActions.cpp b/Engine/Plugins/Experimental/RazerChromaDevices/Source/RazerChromaEditor/Private/RazerChromaAnimationAssetActions.cpp
new file mode 100644
--- /dev/null
+++ b/Engine/Plugins/Experimental/RazerChromaDevices/Source/RazerChromaEditor/Private/RazerChromaAnimationAssetActions.cpp
@@ -0,0 +1,109 @@
+// Copyright Epic Games, Inc. All Rights Reserved.
+
+#include "RazerChromaAnimationAssetActions.h"
+#include "RazerChromaFunctionLibrary.h"
+#include "RazerChromaAnimationAsset.h"
+#include "AssetRegistry/AssetData.h"
+#include "Framework/Commands/UIAction.h"
+
+#include "ToolMenuDelegates.h"
+#include "ToolMenus.h"
+#include "Styling/AppStyle.h"
+
+#define LOCTEXT_NAMESPACE "RazerAssetTypeActions"
+
+FAssetTypeActions_RazerChromaPreviewAction::FAssetTypeActions_RazerChromaPreviewAction(const uint32 InCategoryBit)
+		: CategoryBit(InCategoryBit)
+{
+
+}
+
+UClass* FAssetTypeActions_RazerChromaPreviewAction::GetSupportedClass() const
+{
+	return URazerChromaAnimationAsset::StaticClass();
+}
+
+void FAssetTypeActions_RazerChromaPreviewAction::GetActions(const TArray<UObject*>& InObjects, FToolMenuSection& Section)
+{
+	TArray<TWeakObjectPtr<URazerChromaAnimationAsset>> Anims = GetTypedWeakObjectPtrs<URazerChromaAnimationAsset>(InObjects);
+
+	Section.AddMenuEntry(
+		"RazerChromaAnimPlayEffect",
+		LOCTEXT("RazerChromaAnimPlayEffect", "Play"),
+		LOCTEXT("RazerChromaAnimPlayEffectTooltip", "Plays the selected Razer Chroma Animation"),
+		FSlateIcon(FAppStyle::GetAppStyleSetName(), "Icons.Toolbar.Play"),
+		FUIAction(
+			FExecuteAction::CreateSP(this, &FAssetTypeActions_RazerChromaPreviewAction::ExecutePlayAnim, Anims),
+			FCanExecuteAction::CreateSP(this, &FAssetTypeActions_RazerChromaPreviewAction::CanExecutePlayCommand, Anims)
+		)
+	);
+
+	Section.AddMenuEntry(
+		"RazerChromaAnimStopEffect",
+		LOCTEXT("RazerChromaAnimStopEffect", "Stop"),
+		LOCTEXT("RazerChromaAnimStopEffectTooltip", "Stops the selected Razer Chroma Animation"),
+		FSlateIcon(FAppStyle::GetAppStyleSetName(), "Icons.Toolbar.Stop"),
+		FUIAction(
+			FExecuteAction::CreateSP(this, &FAssetTypeActions_RazerChromaPreviewAction::ExecuteStopAnim, Anims),
+			FCanExecuteAction()
+		)
+	);
+}
+
+bool FAssetTypeActions_RazerChromaPreviewAction::AssetsActivatedOverride(const TArray<UObject*>& InObjects, EAssetTypeActivationMethod::Type ActivationType)
+{
+	return false;
+}
+
+FText FAssetTypeActions_RazerChromaPreviewAction::GetName() const
+{
+	return LOCTEXT("AssetTypeActions_RazerChromaPreviewAction", "Razer Chroma Animation");
+}
+
+FColor FAssetTypeActions_RazerChromaPreviewAction::GetTypeColor() const
+{
+	// Kind of a "Razer Green"
+	return FColor(0, 175, 0);
+}
+
+uint32  FAssetTypeActions_RazerChromaPreviewAction::GetCategories()
+{
+	return CategoryBit;
+}
+
+FText FAssetTypeActions_RazerChromaPreviewAction::GetAssetDescription(const FAssetData& AssetData) const
+{
+	return LOCTEXT("RazerChromaAnimAsset_Desc", "Represents a Razer Chroma animation asset, which can be imported from a '.chroma' file after being created in the Razer Chroma tools suite."); 
+}
+
+void FAssetTypeActions_RazerChromaPreviewAction::ExecutePlayAnim(TArray<TWeakObjectPtr<URazerChromaAnimationAsset>> Objects)
+{
+	for (const TWeakObjectPtr<URazerChromaAnimationAsset>& EffectPtr : Objects)
+	{
+		if (URazerChromaAnimationAsset* Effect = EffectPtr.Get())
+		{
+			URazerChromaFunctionLibrary::PlayChromaAnimation(Effect);
+		}
+	}
+}
+
+void FAssetTypeActions_RazerChromaPreviewAction::ExecuteStopAnim(TArray<TWeakObjectPtr<URazerChromaAnimationAsset>> Objects)
+{
+	for (const TWeakObjectPtr<URazerChromaAnimationAsset>& EffectPtr : Objects)
+	{
+		if (URazerChromaAnimationAsset* Effect = EffectPtr.Get())
+		{
+			// Stop all Razer Chroma animations that are playing
+			URazerChromaFunctionLibrary::StopAllChromaAnimations();
+			break;
+		}
+	}
+}
+
+bool FAssetTypeActions_RazerChromaPreviewAction::CanExecutePlayCommand(TArray<TWeakObjectPtr<URazerChromaAnimationAsset>> Objects) const
+{
+	// We can play a preview animation as long as there is one selected
+	return !Objects.IsEmpty();
+}
+
+#undef LOCTEXT_NAMESPACE
\ No newline at end of file
diff --git a/Engine/Plugins/Experimental/RazerChromaDevices/Source/RazerChromaEditor/Private/RazerChromaAnimationAssetActions.h b/Engine/Plugins/Experimental/RazerChromaDevices/Source/RazerChromaEditor/Private/RazerChromaAnimationAssetActions.h
new file mode 100644
--- /dev/null
+++ b/Engine/Plugins/Experimental/RazerChromaDevices/Source/RazerChromaEditor/Private/RazerChromaAnimationAssetActions.h
@@ -0,0 +1,33 @@
+// Copyright Epic Games, Inc. All Rights Reserved.
+
+#pragma once
+
+#include "AssetTypeActions_Base.h"
+
+class URazerChromaAnimationAsset;
+
+/**
+* Asset Actions for the URazerChromaAnimationAsset which allow you to preview 
+* the animation in the editor without having to PIE. 
+*/
+class FAssetTypeActions_RazerChromaPreviewAction : public FAssetTypeActions_Base
+{
+public:
+	FAssetTypeActions_RazerChromaPreviewAction(const uint32 InCategoryBit);
+
+	virtual FText GetName() const override;
+	virtual FColor GetTypeColor() const override;
+	virtual uint32 GetCategories() override;
+	virtual UClass* GetSupportedClass() const override;
+	virtual void GetActions(const TArray<UObject*>& InObjects, struct FToolMenuSection& Section) override;
+	virtual bool AssetsActivatedOverride(const TArray<UObject*>& InObjects, EAssetTypeActivationMethod::Type ActivationType) override;
+	virtual FText GetAssetDescription(const FAssetData& AssetData) const override;
+
+protected:
+
+	void ExecutePlayAnim(TArray<TWeakObjectPtr<URazerChromaAnimationAsset>> Objects);
+	void ExecuteStopAnim(TArray<TWeakObjectPtr<URazerChromaAnimationAsset>> Objects);
+	bool CanExecutePlayCommand(TArray<TWeakObjectPtr<URazerChromaAnimationAsset>> Objects) const;
+
+	uint32 CategoryBit;
+};
\ No newline at end of file
diff --git a/Engine/Plugins/Experimental/RazerChromaDevices/Source/RazerChromaEditor/Private/RazerChromaEditorModule.cpp b/Engine/Plugins/Experimental/RazerChromaDevices/Source/RazerChromaEditor/Private/RazerChromaEditorModule.cpp
new file mode 100644
--- /dev/null
+++ b/Engine/Plugins/Experimental/RazerChromaDevices/Source/RazerChromaEditor/Private/RazerChromaEditorModule.cpp
@@ -0,0 +1,72 @@
+// Copyright Epic Games, Inc. All Rights Reserved.
+
+#include "Modules/ModuleInterface.h"
+#include "Modules/ModuleManager.h"
+#include "RazerChromaDeviceLogging.h"
+#include "AssetToolsModule.h"
+#include "RazerChromaAnimationAssetActions.h"
+#include "Styling/SlateStyle.h"
+#include "Styling/SlateStyleMacros.h"
+#include "Styling/SlateStyleRegistry.h"
+#include "Interfaces/IPluginManager.h"
+
+/** Custom style set for Razer Chroma */
+class FRazerChromaDevicesSlateStyle final : public FSlateStyleSet
+{
+public:
+	FRazerChromaDevicesSlateStyle()
+		: FSlateStyleSet("RazerChromaDevicesEditor")
+	{
+		SetParentStyleName(FAppStyle::GetAppStyleSetName());
+
+		const FString PluginDirectory = IPluginManager::Get().FindPlugin(TEXT("RazerChromaDevices"))->GetBaseDir();
+		SetContentRoot(PluginDirectory / TEXT("Content/Editor/Slate"));
+		SetCoreContentRoot(FPaths::EngineContentDir() / TEXT("Slate"));
+
+		// Icon sizes
+		const FVector2D Icon16 = FVector2D(16.0, 16.0);
+		const FVector2D Icon64 = FVector2D(64.0, 64.0);
+		
+		Set("ClassIcon.RazerChromaAnimationAsset", new IMAGE_BRUSH_SVG("Icons/ChromaAnimation_16", Icon16));
+		Set("ClassThumbnail.RazerChromaAnimationAsset", new IMAGE_BRUSH_SVG("Icons/ChromaAnimation_64", Icon64));
+	}
+};
+
+class FRazerChromaEditorModule : public IModuleInterface
+{
+protected:
+
+	TSharedPtr<FAssetTypeActions_RazerChromaPreviewAction> RazerChromaPreviewAction;
+
+	TSharedPtr<FSlateStyleSet> StyleSet;
+
+	virtual void StartupModule() override
+	{
+		// Register our custom asset actions for razer chroma animation assets
+		IAssetTools& AssetTools = FModuleManager::LoadModuleChecked<FAssetToolsModule>(TEXT("AssetTools")).Get();
+		EAssetTypeCategories::Type Category = AssetTools.RegisterAdvancedAssetCategory(FName(TEXT("Razer")), NSLOCTEXT("RazerChromaAnimations", "RazerChromaAnimMenu", "Razer Chroma"));
+		RazerChromaPreviewAction = MakeShareable(new FAssetTypeActions_RazerChromaPreviewAction(Category));
+		AssetTools.RegisterAssetTypeActions(RazerChromaPreviewAction.ToSharedRef());
+
+		// Make a new style set for Razer Chroma
+		StyleSet = MakeShared<FRazerChromaDevicesSlateStyle>();
+		FSlateStyleRegistry::RegisterSlateStyle(*StyleSet.Get());
+	}
+	
+	virtual void ShutdownModule() override
+	{
+		if (FModuleManager::Get().IsModuleLoaded(TEXT("AssetTools")))
+		{
+			IAssetTools& AssetTools = FModuleManager::LoadModuleChecked<FAssetToolsModule>("AssetTools").Get();
+			AssetTools.UnregisterAssetTypeActions(RazerChromaPreviewAction.ToSharedRef());
+		}
+
+		// Unregister slate stylings
+		if (StyleSet.IsValid())
+		{
+			FSlateStyleRegistry::UnRegisterSlateStyle(*StyleSet.Get());
+		}
+	}
+};
+
+IMPLEMENT_MODULE(FRazerChromaEditorModule, RazerChromaEditor);
\ No newline at end of file
diff --git a/Engine/Plugins/Experimental/RazerChromaDevices/Source/RazerChromaEditor/Private/RazerChromaFactory.cpp b/Engine/Plugins/Experimental/RazerChromaDevices/Source/RazerChromaEditor/Private/RazerChromaFactory.cpp
new file mode 100644
--- /dev/null
+++ b/Engine/Plugins/Experimental/RazerChromaDevices/Source/RazerChromaEditor/Private/RazerChromaFactory.cpp
@@ -0,0 +1,73 @@
+// Copyright Epic Games, Inc. All Rights Reserved.
+
+#include "RazerChromaFactory.h"
+#include "FileMediaSource.h"
+#include "Misc/Paths.h"
+#include "Misc/MessageDialog.h"
+#include "UObject/UObjectGlobals.h"
+#include "Subsystems/ImportSubsystem.h"
+#include "Editor/EditorEngine.h"
+#include "Editor.h" // for GEditor
+#include "RazerChromaAnimationAsset.h"
+#include "CoreMinimal.h"
+#include "Misc/FeedbackContext.h"
+
+URazerChromaFactory::URazerChromaFactory(const FObjectInitializer& ObjectInitializer)
+	: Super(ObjectInitializer)
+{
+	Formats.Add(TEXT("chroma;Razer Chroma Animation"));
+	bEditorImport = true;
+
+	SupportedClass = URazerChromaAnimationAsset::StaticClass();
+}
+
+bool URazerChromaFactory::FactoryCanImport(const FString& Filename)
+{
+	const FString Extension = FPaths::GetExtension(*Filename);
+	if (Extension == TEXT("chroma"))
+	{
+		return true;
+	}
+	
+	return Super::FactoryCanImport(Filename);
+}
+
+UObject* URazerChromaFactory::FactoryCreateBinary(UClass* InClass, UObject* InParent, FName InName, EObjectFlags Flags, UObject* Context, const TCHAR* Type, const uint8*& Buffer, const uint8* BufferEnd, FFeedbackContext* Warn)
+{
+	GEditor->GetEditorSubsystem<UImportSubsystem>()->BroadcastAssetPreImport(this, InClass, InParent, InName, Type);
+	
+	// Refuse to accept big files. We currently use TArray<> which will fail if we go over an int32.
+	{
+		const uint64 Size = BufferEnd - Buffer;
+		if (!IntFitsIn<int32>(Size))
+		{
+			Warn->Logf(ELogVerbosity::Error, TEXT("File '%s' is too big (%umb), Max=%umb"), *InName.ToString(), Size >> 20, TNumericLimits<int32>::Max() >> 20);
+			GEditor->GetEditorSubsystem<UImportSubsystem>()->BroadcastAssetPostImport(this, nullptr);
+			return nullptr;
+		}
+
+		// This is just an assumption to try and validate binary files here... any animation file should certainly be larger then 4 bytes
+		static constexpr uint64 MinFileSize = 4;
+
+		if (Size <= MinFileSize)
+		{
+			Warn->Logf(ELogVerbosity::Error, TEXT("File '%s' is too small (%u), Min=%u"), *InName.ToString(), Size >> 20, MinFileSize);
+			GEditor->GetEditorSubsystem<UImportSubsystem>()->BroadcastAssetPostImport(this, nullptr);
+			return nullptr;
+		}
+	}
+
+	// Create a new razer animation asset and copy over the byte data for this razer chroma file
+	URazerChromaAnimationAsset* ChromaObject = NewObject<URazerChromaAnimationAsset>(InParent, InName, Flags);
+	const bool bSuccessfulImport = ChromaObject->ImportFromFile(GetCurrentFilename(), Buffer, BufferEnd);
+	
+	if (!bSuccessfulImport)
+	{
+		// Inform user we failed to create the sound wave
+		Warn->Logf(ELogVerbosity::Error, TEXT("Failed to import Razer Chroma Animation %s"), *InName.ToString());
+	}
+
+	GEditor->GetEditorSubsystem<UImportSubsystem>()->BroadcastAssetPostImport(this, ChromaObject);
+	
+	return ChromaObject;
+}
\ No newline at end of file
diff --git a/Engine/Plugins/Experimental/RazerChromaDevices/Source/RazerChromaEditor/Private/RazerChromaFactory.h b/Engine/Plugins/Experimental/RazerChromaDevices/Source/RazerChromaEditor/Private/RazerChromaFactory.h
new file mode 100644
--- /dev/null
+++ b/Engine/Plugins/Experimental/RazerChromaDevices/Source/RazerChromaEditor/Private/RazerChromaFactory.h
@@ -0,0 +1,25 @@
+// Copyright Epic Games, Inc. All Rights Reserved.
+
+#pragma once
+
+#include "UObject/ObjectMacros.h"
+#include "Factories/Factory.h"
+
+#include "RazerChromaFactory.generated.h"
+
+UCLASS(hidecategories=Object)
+class URazerChromaFactory
+	: public UFactory
+{
+	GENERATED_BODY()
+
+public:
+	URazerChromaFactory(const FObjectInitializer& ObjectInitializer);
+	
+protected:
+
+	//~ Begin UFactory Interface
+	virtual bool FactoryCanImport(const FString& Filename) override;
+	virtual UObject* FactoryCreateBinary(UClass* InClass, UObject* InParent, FName InName, EObjectFlags Flags, UObject* Context, const TCHAR* Type, const uint8*& Buffer, const uint8* BufferEnd, FFeedbackContext* Warn) override;
+	//~ End UFactory Interface
+};
diff --git a/Engine/Plugins/Experimental/RazerChromaDevices/Source/RazerChromaEditor/Source/Private/RazerChromaAnimationAssetActions.cpp b/Engine/Plugins/Experimental/RazerChromaDevices/Source/RazerChromaEditor/Source/Private/RazerChromaAnimationAssetActions.cpp
deleted file mode 100644
diff --git a/Engine/Plugins/Experimental/RazerChromaDevices/Source/RazerChromaEditor/Source/Private/RazerChromaAnimationAssetActions.h b/Engine/Plugins/Experimental/RazerChromaDevices/Source/RazerChromaEditor/Source/Private/RazerChromaAnimationAssetActions.h
deleted file mode 100644
diff --git a/Engine/Plugins/Experimental/RazerChromaDevices/Source/RazerChromaEditor/Source/Private/RazerChromaEditorModule.cpp b/Engine/Plugins/Experimental/RazerChromaDevices/Source/RazerChromaEditor/Source/Private/RazerChromaEditorModule.cpp
deleted file mode 100644
diff --git a/Engine/Plugins/Experimental/RazerChromaDevices/Source/RazerChromaEditor/Source/Private/RazerChromaFactory.cpp b/Engine/Plugins/Experimental/RazerChromaDevices/Source/RazerChromaEditor/Source/Private/RazerChromaFactory.cpp
deleted file mode 100644
diff --git a/Engine/Plugins/Experimental/RazerChromaDevices/Source/RazerChromaEditor/Source/Private/RazerChromaFactory.h b/Engine/Plugins/Experimental/RazerChromaDevices/Source/RazerChromaEditor/Source/Private/RazerChromaFactory.h
deleted file mode 100644
